[AI Toolsets]: Move AIModuleToolset under the Toolsets directory

#Synced-CL 52368392

diff --git a/Engine/Plugins/Experimental/AIModuleToolset/AIModuleToolset.uplugin b/Engine/Plugins/Experimental/AIModuleToolset/AIModuleToolset.uplugin
deleted file mode 100644
diff --git a/Engine/Plugins/Experimental/AIModuleToolset/Source/AIModuleToolset/AIModuleToolset.Build.cs b/Engine/Plugins/Experimental/AIModuleToolset/Source/AIModuleToolset/AIModuleToolset.Build.cs
deleted file mode 100644
diff --git a/Engine/Plugins/Experimental/AIModuleToolset/Source/AIModuleToolset/Private/AIModuleToolset.cpp b/Engine/Plugins/Experimental/AIModuleToolset/Source/AIModuleToolset/Private/AIModuleToolset.cpp
deleted file mode 100644
diff --git a/Engine/Plugins/Experimental/AIModuleToolset/Source/AIModuleToolset/Public/AIModuleToolset.h b/Engine/Plugins/Experimental/AIModuleToolset/Source/AIModuleToolset/Public/AIModuleToolset.h
deleted file mode 100644
diff --git a/Engine/Plugins/Experimental/Toolsets/AIModuleToolset/AIModuleToolset.uplugin b/Engine/Plugins/Experimental/Toolsets/AIModuleToolset/AIModuleToolset.uplugin
new file mode 100644
--- /dev/null
+++ b/Engine/Plugins/Experimental/Toolsets/AIModuleToolset/AIModuleToolset.uplugin
@@ -0,0 +1,34 @@
+{
+	"FileVersion": 3,
+	"Version": 1,
+	"VersionName": "1.0",
+	"FriendlyName": "AIModuleToolset",
+	"Description": "Toolset for AIModule Systems",
+	"Category": "AI",
+	"EditorOnly": true,
+	"CreatedBy": "Epic Games",
+	"CreatedByURL": "",
+	"DocsURL": "",
+	"MarketplaceURL": "",
+	"SupportURL": "",
+	"CanContainContent": true,
+	"IsBetaVersion": false,
+	"IsExperimentalVersion": true,
+	"NoRedist": true,
+	"EnabledByDefault": false,
+	"Installed": false,
+	"Modules": [
+		{
+			"Name": "AIModuleToolset",
+			"Type": "Editor",
+			"LoadingPhase": "Default",
+			"TargetAllowList": ["Editor"]
+		}
+	],
+	"Plugins": [
+		{
+			"Name": "ToolsetRegistry",
+			"Enabled": true
+		}
+	]
+}
diff --git a/Engine/Plugins/Experimental/Toolsets/AIModuleToolset/Source/AIModuleToolset/AIModuleToolset.Build.cs b/Engine/Plugins/Experimental/Toolsets/AIModuleToolset/Source/AIModuleToolset/AIModuleToolset.Build.cs
new file mode 100644
--- /dev/null
+++ b/Engine/Plugins/Experimental/Toolsets/AIModuleToolset/Source/AIModuleToolset/AIModuleToolset.Build.cs
@@ -0,0 +1,12 @@
+// Copyright Epic Games, Inc. All Rights Reserved.
+
+using UnrealBuildTool;
+
+public class AIModuleToolset : ModuleRules
+{
+	public AIModuleToolset(ReadOnlyTargetRules Target) : base(Target)
+	{
+		PCHUsage = PCHUsageMode.UseExplicitOrSharedPCHs;
+		PublicDependencyModuleNames.Add("Core");
+	}
+}
diff --git a/Engine/Plugins/Experimental/Toolsets/AIModuleToolset/Source/AIModuleToolset/Private/AIModuleToolset.cpp b/Engine/Plugins/Experimental/Toolsets/AIModuleToolset/Source/AIModuleToolset/Private/AIModuleToolset.cpp
new file mode 100644
--- /dev/null
+++ b/Engine/Plugins/Experimental/Toolsets/AIModuleToolset/Source/AIModuleToolset/Private/AIModuleToolset.cpp
@@ -0,0 +1,16 @@
+// Copyright Epic Games, Inc. All Rights Reserved.
+
+#include "AIModuleToolset.h"
+
+#include "Modules/ModuleManager.h"
+
+/*virtual*/ void FAIModuleToolsetModule::StartupModule()
+{
+}
+
+/*virtual*/ void FAIModuleToolsetModule::ShutdownModule()
+{
+}
+
+IMPLEMENT_MODULE(FAIModuleToolsetModule, AIModuleToolset);
+DEFINE_LOG_CATEGORY(LogAIModuleToolset);
diff --git a/Engine/Plugins/Experimental/Toolsets/AIModuleToolset/Source/AIModuleToolset/Public/AIModuleToolset.h b/Engine/Plugins/Experimental/Toolsets/AIModuleToolset/Source/AIModuleToolset/Public/AIModuleToolset.h
new file mode 100644
--- /dev/null
+++ b/Engine/Plugins/Experimental/Toolsets/AIModuleToolset/Source/AIModuleToolset/Public/AIModuleToolset.h
@@ -0,0 +1,15 @@
+// Copyright Epic Games, Inc. All Rights Reserved.
+
+#pragma once
+
+#include "Modules/ModuleInterface.h"
+#include "Logging/LogMacros.h"
+
+DECLARE_LOG_CATEGORY_EXTERN(LogAIModuleToolset, Log, Log);
+
+class FAIModuleToolsetModule : public IModuleInterface
+{
+public:
+	virtual void StartupModule() override;
+	virtual void ShutdownModule() override;
+};
